Fixing conflicting UBT error logs not being captured due to a change in format, we now only care about "A conflicting instance of" to capture
Bumping ST version

diff --git a/Engine/Programs/SubmitTool/Config/SubmitTool.ini b/Engine/Programs/SubmitTool/Config/SubmitTool.ini
--- a/Engine/Programs/SubmitTool/Config/SubmitTool.ini
+++ b/Engine/Programs/SubmitTool/Config/SubmitTool.ini
@@ -135,7 +135,7 @@ Configurations=("Development", "Debug", "Shipping")
 Platforms=("Win64", "Linux", "Mac")
 Targets=("Editor", "Client")
 ToolTip="This Validator compiles the code in the changelist using the configuration/platform provided, it also runs Static Code Analysis at the same time."
-ErrorMessages=("A conflicting instance of UnrealBuildTool is already Running")
+ErrorMessages=("A conflicting instance of")
 UIGroup="Code Validations"
 ExecutionBlockGroups=("UBT")
 TaskArea="LocalFiles"
diff --git a/Engine/Source/Programs/SubmitTool/Private/Version/AppVersionDefines.h b/Engine/Source/Programs/SubmitTool/Private/Version/AppVersionDefines.h
--- a/Engine/Source/Programs/SubmitTool/Private/Version/AppVersionDefines.h
+++ b/Engine/Source/Programs/SubmitTool/Private/Version/AppVersionDefines.h
@@ -21,7 +21,7 @@
 
 #define SUBMIT_TOOL_VERSION_MAJOR 5
 #define SUBMIT_TOOL_VERSION_MINOR 1
-#define SUBMIT_TOOL_VERSION_PATCH 1
+#define SUBMIT_TOOL_VERSION_PATCH 3
 
 #define SUBMIT_TOOL_VERSION_STRING \
 	VERSION_STRINGIFY(SUBMIT_TOOL_VERSION_MAJOR)\
