[SubmitTool] Add a new local check for common misspelling and typos.

This is based on an extension to the "unacceptable words" tool which adds typo-checking at https://p4-swarm.epicgames.net/reviews/44941366 . This will prevent users from accidentally adding another "seperate", "recieve" or "heirarchy" to the codebase.

diff --git a/Engine/Programs/SubmitTool/Config/SubmitTool.ini b/Engine/Programs/SubmitTool/Config/SubmitTool.ini
--- a/Engine/Programs/SubmitTool/Config/SubmitTool.ini
+++ b/Engine/Programs/SubmitTool/Config/SubmitTool.ini
@@ -101,7 +101,6 @@ InputType="FreeText"
 ;           bAcceptDeletedFiles       -> If this Validator runs on files marked for delete
 ;           bTreatWarningsAsErrors    -> If this Validator should treat warnings as errors
 ;           bUsesIncrementalCache     -> If this validator maintains a local cache of results per file between runs on the same CL, used for incremental validations
-;           bUsesIncrementalCache     -> If this validator maintains a local cache of results per file between runs on the same CL, used for incremental validations
 ;           AdditionalValidationErrorMessages -> Additional error messages to print when this validation fails
 ;           ToolTip                   -> Tooltip when hovering over the Validator
 ;           TaskArea                  -> Area this validator works on, if an area is updated, the validator state will be automatically resetted { Everything, LocalFiles, ShelvedFiles, LocalAndShelvedFiles, Changelist }
