Fixing Mac constant auto-update by removing all the echo that were ending up in the versionfile

diff --git a/Engine/Programs/SubmitTool/Script/Mac/UpdateSubmitTool.sh b/Engine/Programs/SubmitTool/Script/Mac/UpdateSubmitTool.sh
--- a/Engine/Programs/SubmitTool/Script/Mac/UpdateSubmitTool.sh
+++ b/Engine/Programs/SubmitTool/Script/Mac/UpdateSubmitTool.sh
@@ -14,21 +14,8 @@ do
     executableargs="$executableargs $var"
 done
 
-echo "--------------------------------------"
-echo "ZIP: $zip"
-echo "FOLDER: $folder"
-echo "NAME: $name"
-echo "VERSION: $version"
-echo "VERSION FILE: $versionfile"
-echo "EXE: $executable"
-echo "ROOTDIR: $rootdir"
-echo "ARGS: $executableargs"
-echo "--------------------------------------"
-
-echo "Waiting 2 seconds to allow $name to close"
 sleep 2s
 
-echo "Force closing all $name Instances..."
 killall $name
 
 # Not removing the folder to be safe and maintain parity with windows, re-enable if we find a reason
@@ -37,13 +24,11 @@ killall $name
 
 mkdir -p "$folder"
 
-echo "Unzipping $folder"
 unzip -o "$zip" -d "$folder"
 
 rm -f "$versionfile"
 versiondir=${versionfile%/*}
 mkdir -p "$versiondir"
-echo $version > $versionfile
+printf '%s' "$version" > $versionfile
 
-echo "Invoking $executable with arguments: $executableargs"
 nohup "$executable" $executableargs -root-dir "$rootdir" > /dev/null 2>&1 &
\ No newline at end of file
