SubmitTool: Change tag name from Review to ReviewedBy to match the internal config.
Fix tag parsing bugs when running a locally built ST from the source tree.

#rnx

#Synced-CL 49987382

diff --git a/Engine/Programs/SubmitTool/Config/SubmitTool.ini b/Engine/Programs/SubmitTool/Config/SubmitTool.ini
--- a/Engine/Programs/SubmitTool/Config/SubmitTool.ini
+++ b/Engine/Programs/SubmitTool/Config/SubmitTool.ini
@@ -30,7 +30,7 @@ LogsToCollect=(\
 ;           Validation (see above)
 ;
 
-[Tags.Review]
+[Tags.ReviewedBy]
 TagId="#rb"
 TagLabel="Reviewed By"
 ToolTip="The RB (reviewed by) tag will display any reviewers on the change."
