Add more examples to SubmitTool.ini, including an example of building UAT once and using as dependency

diff --git a/Engine/Programs/SubmitTool/Config/SubmitTool.ini b/Engine/Programs/SubmitTool/Config/SubmitTool.ini
--- a/Engine/Programs/SubmitTool/Config/SubmitTool.ini
+++ b/Engine/Programs/SubmitTool/Config/SubmitTool.ini
@@ -113,6 +113,41 @@ IsRequired="true"
 ToolTip="This validator ensures you have all of the required tags in your changelist description."
 TaskArea="Changelist"
 
+; This is an example of having a validator be a prerequisite for other validators.
+; In this case UAT is often used by other validators. Compiling it once in this validator that other validators depend on
+; can be more effiecent as they can be run in parallel with -NoCompileUAT -NoCompile -NoAutoSDK.
+; IncludeFilesWithExtension should be a superset of all dependent filters.
+; See the Validator.CheckBalancedMacros for how to depend on BuildUAT.
+[Validator.BuildUAT]
+Type="CustomValidator"
+IsRequired="false"
+bUsesIncrementalCache="true"
+bLaunchHidden="true"
+bLaunchReallyHidden="true"
+CustomName="Build UAT"
+IncludeFilesWithExtension=(".cc", ".cpp", ".cs", ".h", ".inl", ".m", ".mm", ".usf", ".ush", ".verse", ".versetest", "/Config/Tests/Tags.ini")
+ExecutablePath="$(RunUAT)"
+; For UE 5.7 and prior use the -WaitForUATMutex and -List arguments to compile everything
+ExecutableArguments="-WaitForUATMutex -List"
+; Post UE 5.7 (CL 47029854) use the -CompileOnly and -AutomationProject arguments
+; ExecutableArguments="-WaitForUATMutex -CompileOnly -AutomationProject=AutomationScripts+BuildGraph"
+ToolTip="This builds UAT once to let UAT based validators launch quicker and run in parallel."
+ExecutionBlockGroups=("UBT")
+TaskArea="LocalFiles"
+
+[Validator.AnalyzeSharedDotNet]
+Type="CustomValidator"
+IsRequired="false"
+bTreatWarningsAsErrors="true"
+CustomName="Shared .NET Static Analysis"
+bUsesIncrementalCache="true"
+AppliesToCLRegex="Engine.Source.Programs.Shared"
+ExecutablePath="$(engine)/Build/BatchFiles/RunDotnet$(ScriptExt)"
+ExecutableArguments="build Source/Programs/Shared/Shared.sln -c Analyze -v minimal --no-incremental"
+ToolTip="This validator builds the Shared .NET projects in the Analyze configuration to ensure that no static analysis issues have been introduced."
+ExecutionBlockGroups=("DotNetAnalyze")
+TaskArea="LocalFiles"
+
 [Validator.NonUnityAndStaticAnalysis]
 Type="UBTValidator"
 IsRequired="true"
@@ -139,6 +174,24 @@ UIGroup="Code Validations"
 ExecutionBlockGroups=("UBT")
 TaskArea="LocalFiles"
 
+[Validator.CheckBalancedMacros]
+Type="CustomValidator"
+IsRequired="false"
+bUsesIncrementalCache="true"
+bLaunchHidden="true"
+bLaunchReallyHidden="true"
+bTreatWarningsAsErrors="true"
+CustomName="Balanced Macros Validator"
+FileListArgument="-OverrideFileList="		
+ExecutablePath="$(RunUAT)"
+; The -NoCompileUAT -NoCompile -NoAutoSDK arguments are important as we expect UAT to already be compiled
+ExecutableArguments="CheckBalancedMacros -NoCompileUAT -NoCompile -NoAutoSDK"
+IncludeFilesWithExtension=(".h",".cpp")
+ErrorMessages=("macro appears a second time without matching", "macro appears without matching", "macro does not have matching")
+ToolTip="This validator checks for balanced macros. For example, if you have PRAGMA_DISABLE_OPTIMIZATION without a PRAGMA_ENABLE_OPTIMIZATION"
+DependsOn=("BuildUAT")
+TaskArea="LocalFiles"
+
 [Validator.CrossChangelist]
 Type="CrossChangelistValidator"
 IsRequired="false"
