Added missing .ini to submit tool config

#Synced-CL 50398427

diff --git a/Engine/Programs/SubmitTool/Config/DefaultGame.ini b/Engine/Programs/SubmitTool/Config/DefaultGame.ini
--- a/Engine/Programs/SubmitTool/Config/DefaultGame.ini
+++ b/Engine/Programs/SubmitTool/Config/DefaultGame.ini
@@ -4,6 +4,7 @@
 +AllowedDirectories=Engine/Content/Slate
 +AllowedConfigFiles=Engine/Config/BaseEditorPerProjectUserSettings.ini
 +AllowedConfigFiles=SubmitTool/Config/SubmitTool.ini
++AllowedConfigFiles=SubmitTool/Config/SubmitToolDevOverride.ini
 +AllowedConfigFiles=SubmitTool/Config/DeployId.ini
 +AllowedConfigFiles=SubmitTool/Config/Mac/MacSubmitTool.ini
 +AllowedConfigFiles=SubmitTool/Config/Linux/LinuxSubmitTool.ini
